Consolidated toolbox, gammaui, and module ui into a single debug tools plugin.  The goal of the plugin is to house editor debug UIs going forward.

diff --git a/Engine/Plugins/Editor/EditorDebugTools/EditorDebugTools.uplugin b/Engine/Plugins/Editor/EditorDebugTools/EditorDebugTools.uplugin
new file mode 100644
--- /dev/null
+++ b/Engine/Plugins/Editor/EditorDebugTools/EditorDebugTools.uplugin
@@ -0,0 +1,25 @@
+{
+	"FileVersion": 3,
+	"Version": 1,
+	"VersionName": "1.0",
+	"FriendlyName": "EditorDebugTools",
+	"Description": "",
+	"Category": "Other",
+	"CreatedBy" : "Epic Games, Inc.",
+	"CreatedByURL" : "http://epicgames.com",
+	"DocsURL": "",
+	"MarketplaceURL": "",
+	"SupportURL": "",
+	"EnabledByDefault" : true,
+	"CanContainContent": false,
+	"IsBetaVersion": false,
+	"IsExperimentalVersion": false,
+	"Installed": false,
+	"Modules": [
+		{
+			"Name": "EditorDebugTools",
+			"Type": "Editor",
+			"LoadingPhase": "Default"
+		}
+	]
+}
\ No newline at end of file
diff --git a/Engine/Plugins/Editor/EditorDebugTools/Source/EditorDebugTools/EditorDebugTools.Build.cs b/Engine/Plugins/Editor/EditorDebugTools/Source/EditorDebugTools/EditorDebugTools.Build.cs
new file mode 100644
--- /dev/null
+++ b/Engine/Plugins/Editor/EditorDebugTools/Source/EditorDebugTools/EditorDebugTools.Build.cs
@@ -0,0 +1,65 @@
+// Copyright Epic Games, Inc. All Rights Reserved.
+
+using UnrealBuildTool;
+
+public class EditorDebugTools : ModuleRules
+{
+	public EditorDebugTools(ReadOnlyTargetRules Target) : base(Target)
+	{
+		PCHUsage = ModuleRules.PCHUsageMode.UseExplicitOrSharedPCHs;
+		
+		PublicIncludePaths.AddRange(
+			new string[] {
+			}
+			);
+				
+		
+		PrivateIncludePaths.AddRange(
+			new string[] {
+			}
+			);
+			
+		
+		PublicDependencyModuleNames.AddRange(
+			new string[]
+			{
+				"Core",
+			}
+			);
+			
+		
+		PrivateDependencyModuleNames.AddRange(
+			new string[]
+			{
+				"Projects",
+				"InputCore",
+				"EditorFramework",
+				"UnrealEd",
+				"ToolMenus",
+				"CoreUObject",
+				"Engine",
+				"Slate",
+				"SlateCore",
+				"SourceCodeAccess",
+				"WorkspaceMenuStructure",
+				"AppFramework",
+				"EditorStyle",
+			}
+			);
+
+		if (Target.Configuration != UnrealTargetConfiguration.Shipping)
+		{
+			PrivateIncludePathModuleNames.AddRange(
+				new string[] {
+					"SlateReflector",
+				}
+			);
+
+			DynamicallyLoadedModuleNames.AddRange(
+				new string[] {
+					"SlateReflector",
+				}
+			);
+		}
+	}
+}
diff --git a/Engine/Plugins/Editor/EditorDebugTools/Source/EditorDebugTools/Private/EditorDebugTools.cpp b/Engine/Plugins/Editor/EditorDebugTools/Source/EditorDebugTools/Private/EditorDebugTools.cpp
new file mode 100644
--- /dev/null
+++ b/Engine/Plugins/Editor/EditorDebugTools/Source/EditorDebugTools/Private/EditorDebugTools.cpp
@@ -0,0 +1,96 @@
+// Copyright Epic Games, Inc. All Rights Reserved.
+
+#include "EditorDebugTools.h"
+#include "EditorDebugToolsStyle.h"
+#include "LevelEditor.h"
+#include "Widgets/Docking/SDockTab.h"
+#include "Widgets/Layout/SBox.h"
+#include "Widgets/Text/STextBlock.h"
+#include "ToolMenus.h"
+#include "SDebugPanel.h"
+#include "ISourceCodeAccessModule.h"
+#include "ISourceCodeAccessor.h"
+#include "GammaUIPanel.h"
+#include "SModuleUI.h"
+#include "WorkspaceMenuStructureModule.h"
+#include "WorkspaceMenuStructure.h"
+
+static const FName EditorDebugToolsTabName("DebugTools");
+static const FName ModulesTabName("Modules");
+
+static bool CanShowModulesTab()
+{
+	ISourceCodeAccessModule* SourceCodeAccessModule = FModuleManager::GetModulePtr<ISourceCodeAccessModule>("SourceCodeAccess");
+	return SourceCodeAccessModule != nullptr && SourceCodeAccessModule->GetAccessor().CanAccessSourceCode();
+}
+
+#define LOCTEXT_NAMESPACE "FEditorDebugToolsModule"
+
+TSharedRef<SDockTab> CreateDebugToolsTab(const FSpawnTabArgs& Args)
+{
+	return
+		SNew(SDockTab)
+		.TabRole(ETabRole::NomadTab)
+		[
+			SNew(SVerticalBox)
+			+ SVerticalBox::Slot()
+			.AutoHeight()
+			[
+				SNew(SDebugPanel)
+			]
+			+ SVerticalBox::Slot()
+			.AutoHeight()
+			[
+				SNew(SGammaUIPanel)
+			]
+		];
+}
+
+TSharedRef<SDockTab> CreateModulesTab(const FSpawnTabArgs& Args)
+{
+	return
+		SNew(SDockTab)
+		.TabRole(ETabRole::NomadTab)
+		[
+			SNew(SBorder)
+			.Padding(2)
+			.BorderImage(FEditorStyle::GetBrush("ToolPanel.GroupBorder"))
+			[
+				SNew(SModuleUI)
+			]
+		];
+}
+
+void FEditorDebugToolsModule::StartupModule()
+{
+	FTabSpawnerEntry& Spawner = FGlobalTabmanager::Get()->RegisterNomadTabSpawner(EditorDebugToolsTabName, FOnSpawnTab::CreateStatic(&CreateDebugToolsTab))
+		.SetDisplayName(NSLOCTEXT("Toolbox", "DebugTools", "Debug Tools"))
+		.SetTooltipText(NSLOCTEXT("Toolbox", "DebugToolsTooltipText", "Open the Debug Tools tab."))
+		.SetIcon(FSlateIcon(FEditorStyle::GetStyleSetName(), "DebugTools.TabIcon"))
+		.SetGroup(WorkspaceMenu::GetMenuStructure().GetDeveloperToolsDebugCategory());
+
+	if (CanShowModulesTab())
+	{
+		FGlobalTabmanager::Get()->RegisterNomadTabSpawner(ModulesTabName, FOnSpawnTab::CreateStatic(&CreateModulesTab))
+			.SetDisplayName(NSLOCTEXT("Toolbox", "Modules", "Modules"))
+			.SetTooltipText(NSLOCTEXT("Toolbox", "ModulesTooltipText", "Open the Modules tab."))
+			.SetIcon(FSlateIcon(FEditorStyle::GetStyleSetName(), "Modules.TabIcon"))
+			.SetGroup(WorkspaceMenu::GetMenuStructure().GetDeveloperToolsMiscCategory());
+	}
+}
+
+void FEditorDebugToolsModule::ShutdownModule()
+{
+	FGlobalTabmanager::Get()->UnregisterNomadTabSpawner(EditorDebugToolsTabName);
+
+	if (CanShowModulesTab())
+	{
+		FGlobalTabmanager::Get()->UnregisterNomadTabSpawner(ModulesTabName);
+	}
+}
+
+
+
+#undef LOCTEXT_NAMESPACE
+	
+IMPLEMENT_MODULE(FEditorDebugToolsModule, EditorDebugTools)
\ No newline at end of file
diff --git a/Engine/Plugins/Editor/EditorDebugTools/Source/EditorDebugTools/Private/EditorDebugToolsStyle.cpp b/Engine/Plugins/Editor/EditorDebugTools/Source/EditorDebugTools/Private/EditorDebugToolsStyle.cpp
new file mode 100644
--- /dev/null
+++ b/Engine/Plugins/Editor/EditorDebugTools/Source/EditorDebugTools/Private/EditorDebugToolsStyle.cpp
@@ -0,0 +1,50 @@
+// Copyright Epic Games, Inc. All Rights Reserved.
+
+#include "EditorDebugToolsStyle.h"
+#include "Styling/SlateStyleRegistry.h"
+#include "Framework/Application/SlateApplication.h"
+#include "Interfaces/IPluginManager.h"
+
+TSharedPtr< FSlateStyleSet > FEditorDebugToolsStyle::StyleInstance = NULL;
+
+void FEditorDebugToolsStyle::Initialize()
+{
+	if (!StyleInstance.IsValid())
+	{
+		StyleInstance = Create();
+		FSlateStyleRegistry::RegisterSlateStyle(*StyleInstance);
+	}
+}
+
+void FEditorDebugToolsStyle::Shutdown()
+{
+	FSlateStyleRegistry::UnRegisterSlateStyle(*StyleInstance);
+	ensure(StyleInstance.IsUnique());
+	StyleInstance.Reset();
+}
+
+FName FEditorDebugToolsStyle::GetStyleSetName()
+{
+	static FName StyleSetName(TEXT("EditorDebugToolsStyle"));
+	return StyleSetName;
+}
+
+
+const FVector2D Icon16x16(16.0f, 16.0f);
+const FVector2D Icon20x20(20.0f, 20.0f);
+const FVector2D Icon40x40(40.0f, 40.0f);
+
+TSharedRef< FSlateStyleSet > FEditorDebugToolsStyle::Create()
+{
+	TSharedRef< FSlateStyleSet > Style = MakeShareable(new FSlateStyleSet("EditorDebugToolsStyle"));
+	Style->SetContentRoot(IPluginManager::Get().FindPlugin("EditorDebugTools")->GetBaseDir() / TEXT("Resources"));
+
+	return Style;
+}
+
+
+
+const ISlateStyle& FEditorDebugToolsStyle::Get()
+{
+	return *StyleInstance;
+}
diff --git a/Engine/Plugins/Editor/EditorDebugTools/Source/EditorDebugTools/Private/GammaUIPanel.cpp b/Engine/Plugins/Editor/EditorDebugTools/Source/EditorDebugTools/Private/GammaUIPanel.cpp
new file mode 100644
--- /dev/null
+++ b/Engine/Plugins/Editor/EditorDebugTools/Source/EditorDebugTools/Private/GammaUIPanel.cpp
@@ -0,0 +1,61 @@
+// Copyright Epic Games, Inc. All Rights Reserved.
+
+#include "GammaUIPanel.h"
+#include "EngineGlobals.h"
+#include "Engine/Engine.h"
+#include "Widgets/SBoxPanel.h"
+#include "Widgets/Images/SImage.h"
+#include "Widgets/Text/STextBlock.h"
+#include "Widgets/Input/SSpinBox.h"
+#include "EditorStyleSet.h"
+
+void SGammaUIPanel::Construct(const SGammaUIPanel::FArguments& InArgs)
+{
+	ChildSlot
+		.Padding( FMargin(8) )
+	[
+		SNew( SVerticalBox )
+
+		+ SVerticalBox::Slot()
+			.AutoHeight()
+			.Padding( 0.0f, 4.0f, 0.0f, 4.0f )
+		[
+			SNew( STextBlock )
+				.Text( NSLOCTEXT( "GammaUI", "GammaUILabel", "Gamma" ) )
+		]
+
+		+ SVerticalBox::Slot()
+			.AutoHeight()
+			.Padding( 0.0f, 4.0f, 0.0f, 4.0f )
+		[
+			SNew( SSpinBox<float> )
+			.Delta(0.01f)
+			.MinValue(1.0f)
+			.MaxValue(3.0f)
+			.Value( this, &SGammaUIPanel::OnGetGamma )
+			.OnValueChanged( this, &SGammaUIPanel::OnGammaChanged )
+		]
+
+		+ SVerticalBox::Slot()
+			.AutoHeight()
+			.Padding( 0.0f, 4.0f, 0.0f, 4.0f )
+			.HAlign( HAlign_Center )
+		[
+			SNew(SImage)
+			.Image(FEditorStyle::GetBrush(TEXT("GammaReference")))
+		]
+	];
+}
+
+float SGammaUIPanel::OnGetGamma() const
+{
+	return GEngine ? GEngine->DisplayGamma : 2.2f;
+}
+
+void SGammaUIPanel::OnGammaChanged(float NewValue)
+{
+	if( GEngine )
+	{
+		GEngine->DisplayGamma = NewValue;
+	}
+}
diff --git a/Engine/Plugins/Editor/EditorDebugTools/Source/EditorDebugTools/Private/GammaUIPanel.h b/Engine/Plugins/Editor/EditorDebugTools/Source/EditorDebugTools/Private/GammaUIPanel.h
new file mode 100644
--- /dev/null
+++ b/Engine/Plugins/Editor/EditorDebugTools/Source/EditorDebugTools/Private/GammaUIPanel.h
@@ -0,0 +1,21 @@
+// Copyright Epic Games, Inc. All Rights Reserved.
+
+#pragma once
+
+#include "CoreMinimal.h"
+#include "Widgets/DeclarativeSyntaxSupport.h"
+#include "Widgets/SCompoundWidget.h"
+
+class SGammaUIPanel : public SCompoundWidget
+{
+
+public:
+
+	SLATE_BEGIN_ARGS( SGammaUIPanel ){}
+	SLATE_END_ARGS()
+
+	void Construct(const FArguments& InArgs);
+
+	float OnGetGamma() const;
+	void OnGammaChanged( float NewValue );
+};
diff --git a/Engine/Plugins/Editor/EditorDebugTools/Source/EditorDebugTools/Private/SDebugPanel.cpp b/Engine/Plugins/Editor/EditorDebugTools/Source/EditorDebugTools/Private/SDebugPanel.cpp
new file mode 100644
--- /dev/null
+++ b/Engine/Plugins/Editor/EditorDebugTools/Source/EditorDebugTools/Private/SDebugPanel.cpp
@@ -0,0 +1,104 @@
+#include "SDebugPanel.h"
+#include "ISlateReflectorModule.h"
+#include "Widgets/Testing/SStarshipSuite.h"
+#include "Widgets/Testing/STestSuite.h"
+
+
+void SDebugPanel::Construct(const FArguments& InArgs)
+{
+	ChildSlot
+	[
+		SNew( SVerticalBox )
+		+ SVerticalBox::Slot()
+		.AutoHeight()
+		.Padding( 4.0f )
+		.HAlign(HAlign_Left)
+		[
+			SNew( SButton )
+			.Text( NSLOCTEXT("DeveloperToolbox", "ReloadTextures", "Reload Textures") )
+			.OnClicked( this, &SDebugPanel::OnReloadTexturesClicked )
+		]
+			
+		+ SVerticalBox::Slot()
+		.AutoHeight()
+		.Padding( 4.0f )
+		.HAlign(HAlign_Left)
+		[
+			SNew( SButton )
+			.Text( NSLOCTEXT("DeveloperToolbox", "FlushFontCache", "Flush Font Cache") )
+			.OnClicked( this, &SDebugPanel::OnFlushFontCacheClicked )
+		]
+		+ SVerticalBox::Slot()
+		.AutoHeight()
+		.Padding( 4.0f )
+		.HAlign(HAlign_Left)
+		[
+			SNew( SButton )
+			.Text( NSLOCTEXT("DeveloperToolbox", "TestSuite", "Test Suite") )
+			.OnClicked( this, &SDebugPanel::OnTestSuiteClicked )
+		]
+		+ SVerticalBox::Slot()
+		.AutoHeight()
+		.Padding( 4.0f )
+		.HAlign(HAlign_Left)
+		[
+			SNew( SButton )
+			.Text( NSLOCTEXT("DeveloperToolbox", "DisplayTextureAtlases", "Display Texture Atlases") )
+			.OnClicked( this, &SDebugPanel::OnDisplayTextureAtlases )
+		]
+		+ SVerticalBox::Slot()
+		.AutoHeight()
+		.Padding( 4.0f )
+		.HAlign(HAlign_Left)
+		[
+			SNew( SButton )
+			.Text( NSLOCTEXT("DeveloperToolbox", "DisplayFontAtlases", "Display Font Atlases") )
+			.OnClicked( this, &SDebugPanel::OnDisplayFontAtlases )
+		]
+	];
+}
+
+FReply SDebugPanel::OnReloadTexturesClicked()
+{
+	FSlateApplication::Get().GetRenderer()->ReloadTextureResources();
+	return FReply::Handled();
+}
+
+FReply SDebugPanel::OnDisplayTextureAtlases()
+{
+#if !UE_BUILD_SHIPPING
+	static const FName SlateReflectorModuleName("SlateReflector");
+	FModuleManager::LoadModuleChecked<ISlateReflectorModule>(SlateReflectorModuleName).DisplayTextureAtlasVisualizer();
+#endif // #if !UE_BUILD_SHIPPING
+	return FReply::Handled();
+}
+
+FReply SDebugPanel::OnDisplayFontAtlases()
+{
+#if !UE_BUILD_SHIPPING
+	static const FName SlateReflectorModuleName("SlateReflector");
+	FModuleManager::LoadModuleChecked<ISlateReflectorModule>(SlateReflectorModuleName).DisplayFontAtlasVisualizer();
+#endif // #if !UE_BUILD_SHIPPING
+	return FReply::Handled();
+}
+
+FReply SDebugPanel::OnFlushFontCacheClicked()
+{
+	FSlateApplication::Get().GetRenderer()->FlushFontCache(TEXT("SDebugPanel::OnFlushFontCacheClicked"));
+	return FReply::Handled();
+}
+
+FReply SDebugPanel::OnTestSuiteClicked()
+{
+#if !UE_BUILD_SHIPPING
+	if (FCoreStyle::IsStarshipStyle())
+	{
+		RestoreStarshipSuite();
+	}
+	else
+	{
+		RestoreSlateTestSuite();
+	}
+#endif
+	return FReply::Handled();
+}
\ No newline at end of file
diff --git a/Engine/Plugins/Editor/EditorDebugTools/Source/EditorDebugTools/Private/SDebugPanel.h b/Engine/Plugins/Editor/EditorDebugTools/Source/EditorDebugTools/Private/SDebugPanel.h
new file mode 100644
--- /dev/null
+++ b/Engine/Plugins/Editor/EditorDebugTools/Source/EditorDebugTools/Private/SDebugPanel.h
@@ -0,0 +1,25 @@
+// Copyright Epic Games, Inc. All Rights Reserved.
+
+#pragma once
+
+#include "CoreMinimal.h"
+#include "Widgets/DeclarativeSyntaxSupport.h"
+#include "Widgets/SCompoundWidget.h"
+
+class SDebugPanel : public SCompoundWidget
+{
+public:
+
+	SLATE_BEGIN_ARGS(SDebugPanel) {}
+	SLATE_END_ARGS()
+
+	void Construct(const FArguments& InArgs);
+
+private:
+	FReply OnReloadTexturesClicked();
+	FReply OnDisplayTextureAtlases();
+	FReply OnDisplayFontAtlases();
+	FReply OnFlushFontCacheClicked();
+	FReply OnTestSuiteClicked();
+};
+
diff --git a/Engine/Plugins/Editor/EditorDebugTools/Source/EditorDebugTools/Private/SModuleUI.cpp b/Engine/Plugins/Editor/EditorDebugTools/Source/EditorDebugTools/Private/SModuleUI.cpp
new file mode 100644
--- /dev/null
+++ b/Engine/Plugins/Editor/EditorDebugTools/Source/EditorDebugTools/Private/SModuleUI.cpp
@@ -0,0 +1,345 @@
+// Copyright Epic Games, Inc. All Rights Reserved.
+
+#include "SModuleUI.h"
+#include "HAL/PlatformFileManager.h"
+#include "Misc/FeedbackContext.h"
+#include "GenericPlatform/GenericPlatformFile.h"
+#include "Modules/ModuleManager.h"
+#include "EngineGlobals.h"
+#include "Engine/Engine.h"
+#include "SlateOptMacros.h"
+#include "Widgets/Text/STextBlock.h"
+#include "Widgets/Input/SButton.h"
+#include "Widgets/Views/SListView.h"
+#include "Misc/HotReloadInterface.h"
+#include "Widgets/Input/SSearchBox.h"
+
+BEGIN_SLATE_FUNCTION_BUILD_OPTIMIZATION
+void SModuleUI::Construct(const SModuleUI::FArguments& InArgs)
+{
+	#define LOCTEXT_NAMESPACE "ModuleUI"
+
+	ChildSlot
+		.Padding( FMargin(8) )
+	[
+		SNew( SVerticalBox )
+
+			+ SVerticalBox::Slot()
+				.AutoHeight()
+			[
+				SAssignNew(ModuleNameSearchBox, SSearchBox)
+				.OnTextChanged(this, &SModuleUI::OnFilterTextChanged)
+			]
+
+			// List of modules
+			+ SVerticalBox::Slot()
+				.FillHeight( 1.0f )		// We want the list to stretch vertically to fill up the user-resizable space
+			[
+				SAssignNew( ModuleListView, SModuleListView )
+					.ItemHeight( 24 )
+					.ListItemsSource( &ModuleListItems )
+					.OnGenerateRow( this, &SModuleUI::OnGenerateWidgetForModuleListView )
+					.HeaderRow
+					(
+						SNew(SHeaderRow)
+						+SHeaderRow::Column("ModuleName")
+						.DefaultLabel(NSLOCTEXT("ModuleUI", "ModuleName", "Module"))
+						.FillWidth(200)
+						+SHeaderRow::Column("ModuleActions")
+						.DefaultLabel(NSLOCTEXT("ModuleUI", "ModuleActions", "Actions"))
+						.FillWidth(1000)
+					)
+			]
+	];
+
+	#undef LOCTEXT_NAMESPACE
+
+
+	// Register to find out about module changes
+	FModuleManager::Get().OnModulesChanged().AddSP( this, &SModuleUI::OnModulesChanged );
+
+
+	// Gather data from module manager
+	UpdateModuleListItems();
+}
+END_SLATE_FUNCTION_BUILD_OPTIMIZATION
+
+
+SModuleUI::~SModuleUI()
+{
+	// Unregister callbacks
+	FModuleManager::Get().OnModulesChanged().RemoveAll( this );
+}
+
+void SModuleUI::OnFilterTextChanged(const FText& InFilterText)
+{
+	UpdateModuleListItems();
+}
+TSharedRef<ITableRow> SModuleUI::OnGenerateWidgetForModuleListView(TSharedPtr< FModuleListItem > InItem, const TSharedRef<STableViewBase>& OwnerTable)
+{
+	#define LOCTEXT_NAMESPACE "ModuleUI"
+
+	class SModuleItemWidget : public SMultiColumnTableRow< TSharedPtr< FModuleListItem > >
+	{
+		public:
+		SLATE_BEGIN_ARGS( SModuleItemWidget ){}
+		SLATE_END_ARGS()
+
+		void Construct( const FArguments& InArgs, const TSharedRef<STableViewBase>& InOwnerTable, TSharedPtr<FModuleListItem> InListItem )
+		{
+			Item = InListItem;
+	
+			SMultiColumnTableRow< TSharedPtr< FModuleListItem > >::Construct( FSuperRowType::FArguments(), InOwnerTable );			
+		}
+
+		TSharedRef<SWidget> GenerateWidgetForColumn( const FName& ColumnName )
+		{
+			if ( ColumnName == "ModuleName" )
+			{
+				return
+					SNew( STextBlock )
+					.Text( FText::FromName(Item->ModuleName) );
+			}
+			else if ( ColumnName == "ModuleActions" )
+			{
+				return 
+				SNew( SHorizontalBox )
+	
+				// Load button
+				+ SHorizontalBox::Slot()
+					.AutoWidth()
+					.Padding( 2.0f, 0.0f )
+				[
+					SNew( SButton )
+						.Visibility( Item.ToSharedRef(), &FModuleListItem::GetVisibilityBasedOnUnloadedState )
+						.Text( LOCTEXT("Load", "Load") )
+						.OnClicked( Item.ToSharedRef(), &FModuleListItem::OnLoadClicked )
+				]
+
+				// Unload button
+				+ SHorizontalBox::Slot()
+					.AutoWidth()
+					.Padding( 2.0f, 0.0f )
+				[
+					SNew( SButton )
+						.Visibility( Item.ToSharedRef(), &FModuleListItem::GetVisibilityBasedOnLoadedAndShutdownableState )
+						.Text( LOCTEXT("Unload", "Unload") )
+						.OnClicked( Item.ToSharedRef(), &FModuleListItem::OnUnloadClicked )
+				]
+
+				// Reload button
+				+ SHorizontalBox::Slot()
+					.AutoWidth()
+					.Padding( 2.0f, 0.0f )
+					[
+						SNew( SButton )
+						.Visibility( Item.ToSharedRef(), &FModuleListItem::GetVisibilityBasedOnReloadableState )
+						.Text( LOCTEXT("Reload", "Reload") )
+						.OnClicked( Item.ToSharedRef(), &FModuleListItem::OnReloadClicked )
+					]
+
+				// Recompile button
+				+ SHorizontalBox::Slot()
+					.AutoWidth()
+					.Padding( 2.0f, 0.0f )
+				[
+					SNew( SButton )
+						.Visibility( Item.ToSharedRef(), &FModuleListItem::GetVisibilityBasedOnRecompilableState )
+						.Text( LOCTEXT("Recompile", "Recompile") )
+						.OnClicked( Item.ToSharedRef(), &FModuleListItem::OnRecompileClicked )
+				];
+			}
+			else
+			{
+				return SNew(STextBlock) .Text(LOCTEXT("UnknownColumn", "Unknown Column"));
+			}
+
+		}
+
+		TSharedPtr< FModuleListItem > Item;
+	};
+
+	return SNew( SModuleItemWidget, OwnerTable, InItem );
+
+	#undef LOCTEXT_NAMESPACE
+}
+
+
+void SModuleUI::OnModulesChanged( FName ModuleThatChanged, EModuleChangeReason ReasonForChange )
+{
+	// @todo: Consider using dirty bit for this instead, refresh on demand
+	UpdateModuleListItems();
+}
+
+void SModuleUI::UpdateModuleListItems()
+{
+	ModuleListItems.Reset();
+
+	TArray< FModuleStatus > ModuleStatuses;
+	FModuleManager::Get().QueryModules( ModuleStatuses );
+
+	for( TArray< FModuleStatus >::TConstIterator ModuleIt( ModuleStatuses ); ModuleIt; ++ModuleIt )
+	{
+		const FModuleStatus& ModuleStatus = *ModuleIt;
+
+		FName ModuleName(*ModuleStatus.Name);
+
+		FString ModuleNameStr = ModuleName.ToString();
+		FString FilterStr = ModuleNameSearchBox->GetText().ToString();
+		FilterStr.TrimStartAndEndInline();
+		if ( (!FilterStr.IsEmpty() && ModuleNameStr.Contains(FilterStr)) || 
+			  FilterStr.IsEmpty() )
+		{
+			TSharedRef< FModuleListItem > NewItem(new FModuleListItem());
+			NewItem->ModuleName = ModuleName;
+			ModuleListItems.Add(NewItem);
+		}
+	}
+
+
+	// Sort the list of modules alphabetically
+	struct FModuleSorter
+	{
+		FORCEINLINE bool operator()( const TSharedPtr<FModuleListItem>& A, const TSharedPtr<FModuleListItem>& B ) const
+		{
+			return A->ModuleName.LexicalLess(B->ModuleName);
+		}
+	};
+	ModuleListItems.Sort( FModuleSorter() );
+
+	// Update the list view if we have one
+	if( ModuleListView.IsValid() )
+	{
+		ModuleListView->RequestListRefresh();
+	}
+}
+
+
+FReply SModuleUI::FModuleListItem::OnLoadClicked()
+{
+	GEngine->DeferredCommands.Add( FString::Printf( TEXT( "Module Load %s" ), *ModuleName.ToString() ) );
+	return FReply::Handled();
+}
+
+
+FReply SModuleUI::FModuleListItem::OnUnloadClicked()
+{
+	GEngine->DeferredCommands.Add( FString::Printf( TEXT( "Module Unload %s" ), *ModuleName.ToString() ) );
+	return FReply::Handled();
+}
+
+
+FReply SModuleUI::FModuleListItem::OnReloadClicked()
+{
+	GEngine->DeferredCommands.Add( FString::Printf( TEXT( "Module Reload %s" ), *ModuleName.ToString() ) );
+	return FReply::Handled();
+}
+
+
+FReply SModuleUI::FModuleListItem::OnRecompileClicked()
+{
+	const bool bShowProgressDialog = true;
+	const bool bShowCancelButton = false;
+
+	FFormatNamedArguments Args;
+	Args.Add( TEXT("ModuleName"), FText::FromName( ModuleName ) );
+	GWarn->BeginSlowTask( FText::Format( NSLOCTEXT("ModuleUI", "Recompile_SlowTaskName", "Compiling {ModuleName}..."), Args ), bShowProgressDialog, bShowCancelButton );
+
+	// Does the module have any UObject classes in it?  If so we'll use HotReload to recompile it.
+	FModuleStatus ModuleStatus;
+	if( ensure( FModuleManager::Get().QueryModule( ModuleName, ModuleStatus ) ) )
+	{
+		//@todo This is for content-only packages that show up in the
+		// Module UI... don't crash when recompile is clicked
+		if (FPaths::IsProjectFilePathSet())
+		{
+			if (ModuleStatus.bIsLoaded == false)
+			{
+				if (FPlatformFileManager::Get().GetPlatformFile().FileExists(*ModuleStatus.FilePath) == false)
+				{
+					UE_LOG(LogTemp, Display, TEXT("Unable to recompile module %s... Is it a content-only module?"), *ModuleName.ToString());
+					GWarn->EndSlowTask();
+					return FReply::Handled();
+				}
+			}
+		}
+
+		TArray< UPackage* > PackagesToRebind;
+		if( ModuleStatus.bIsLoaded )
+		{
+			const bool bIsHotReloadable = FModuleManager::Get().DoesLoadedModuleHaveUObjects( ModuleName );
+			if( bIsHotReloadable )
+			{
+				// Is there a UPackage with the same name as this module?
+				FString PotentialPackageName = FString(TEXT("/Script/")) + ModuleName.ToString();
+				UPackage* Package = FindPackage( NULL, *PotentialPackageName);
+				if( Package != NULL )
+				{
+					PackagesToRebind.Add( Package );
+				}
+			}
+		}
+
+		IHotReloadInterface& HotReloadSupport = FModuleManager::LoadModuleChecked<IHotReloadInterface>("HotReload");
+		if( PackagesToRebind.Num() > 0 )
+		{
+			// Perform a hot reload
+			ECompilationResult::Type Result = HotReloadSupport.RebindPackages(PackagesToRebind, EHotReloadFlags::WaitForCompletion, *GLog);
+		}
+		else
+		{
+			// Perform a regular unload, then reload
+			const bool bRecompileSucceeded = HotReloadSupport.RecompileModule(ModuleName, *GLog, ERecompileModuleFlags::ReloadAfterRecompile | ERecompileModuleFlags::FailIfGeneratedCodeChanges);
+		}
+	}
+
+	GWarn->EndSlowTask();
+
+	return FReply::Handled();
+}
+
+
+EVisibility SModuleUI::FModuleListItem::GetVisibilityBasedOnLoadedAndShutdownableState() const
+{
+	if ( GIsSavingPackage || IsGarbageCollecting() )
+	{
+		return EVisibility::Hidden;
+	}
+
+	const bool bIsHotReloadable = FModuleManager::Get().DoesLoadedModuleHaveUObjects(ModuleName);
+	const bool bCanShutDown = ( FModuleManager::Get().IsModuleLoaded(ModuleName)
+								&& !bIsHotReloadable 
+								&& FModuleManager::Get().GetModule(ModuleName)->SupportsDynamicReloading() );
+
+	return bCanShutDown ? EVisibility::Visible : EVisibility::Hidden;
+}
+
+
+EVisibility SModuleUI::FModuleListItem::GetVisibilityBasedOnReloadableState() const
+{
+	return GetVisibilityBasedOnLoadedAndShutdownableState();
+};
+
+
+EVisibility SModuleUI::FModuleListItem::GetVisibilityBasedOnRecompilableState() const
+{
+	if ( GIsSavingPackage || IsGarbageCollecting() )
+	{
+		return EVisibility::Hidden;
+	}
+
+	const bool bIsHotReloadable = FModuleManager::Get().DoesLoadedModuleHaveUObjects(ModuleName);
+	const bool bCanReload = ( !FModuleManager::Get().IsModuleLoaded(ModuleName)
+							|| FModuleManager::Get().GetModule(ModuleName)->SupportsDynamicReloading()
+							|| bIsHotReloadable );
+
+	return bCanReload ? EVisibility::Visible : EVisibility::Hidden;
+}
+
+
+EVisibility SModuleUI::FModuleListItem::GetVisibilityBasedOnUnloadedState() const
+{
+	return FModuleManager::Get().IsModuleLoaded( ModuleName ) ? EVisibility::Hidden : EVisibility::Visible;
+}
+
+
diff --git a/Engine/Plugins/Editor/EditorDebugTools/Source/EditorDebugTools/Private/SModuleUI.h b/Engine/Plugins/Editor/EditorDebugTools/Source/EditorDebugTools/Private/SModuleUI.h
new file mode 100644
--- /dev/null
+++ b/Engine/Plugins/Editor/EditorDebugTools/Source/EditorDebugTools/Private/SModuleUI.h
@@ -0,0 +1,110 @@
+// Copyright Epic Games, Inc. All Rights Reserved.
+
+#pragma once
+
+#include "CoreMinimal.h"
+#include "SlateFwd.h"
+#include "Layout/Visibility.h"
+#include "Input/Reply.h"
+#include "Widgets/DeclarativeSyntaxSupport.h"
+#include "Widgets/SCompoundWidget.h"
+#include "Widgets/Views/STableViewBase.h"
+#include "Widgets/Views/STableRow.h"
+
+template <typename ItemType> class SListView;
+
+class SModuleUI : public SCompoundWidget
+{
+
+public:
+
+	SLATE_BEGIN_ARGS( SModuleUI ){}
+	SLATE_END_ARGS()
+
+
+	/**
+	 * Constructs this widget
+	 *
+	 * @param InArgs    Declaration from which to construct the widget
+	 */
+	void Construct(const FArguments& InArgs);
+
+	/** SModuleUI destructor */
+	~SModuleUI();
+
+
+private:
+
+	/** An item in the module list */
+	struct FModuleListItem
+	{
+		/** Name of this module */
+		FName ModuleName;
+
+		/** Called when 'Load' is clicked in the UI */
+		FReply OnLoadClicked();
+
+		/** Called when 'Unload' is clicked in the UI */
+		FReply OnUnloadClicked();
+
+		/** Called when 'Reload' is clicked in the UI */
+		FReply OnReloadClicked();
+
+		/** Called when 'Recompile' is clicked in the UI */
+		FReply OnRecompileClicked();
+
+		/** @return Returns Visible for modules that are currently loaded and shutdown-able, otherwise Collapsed */
+		EVisibility GetVisibilityBasedOnLoadedAndShutdownableState() const;
+
+		/** @return Returns Visible for modules that are currently loaded and reload-able, otherwise Collapsed */
+		EVisibility GetVisibilityBasedOnReloadableState() const;
+
+		/** @return Returns Visible for modules that are either not yet loaded or can be shutdown, otherwise Collapsed */
+		EVisibility GetVisibilityBasedOnRecompilableState() const;
+
+		/** @return Returns Collapsed for modules that are currently loaded, otherwise Visible */
+		EVisibility GetVisibilityBasedOnUnloadedState() const;
+
+	};
+
+	/**
+	 * Given a data item, reutrns a widget to represent that item in the list view
+	 *
+	 * @param InItem      The list item to return a widget for
+	 * @param OwnerTable  The TableView that this TableRow is being generated for
+	 *
+	 * @return	The widget to represent the specified list item	 
+	 */
+	TSharedRef<ITableRow> OnGenerateWidgetForModuleListView(TSharedPtr< FModuleListItem > InItem, const TSharedRef<STableViewBase>& OwnerTable);
+
+	/**
+	 * Called by the engine's module manager when a module is loaded, unloaded, or the list of known
+	 * modules has changed
+	 */
+	void OnModulesChanged( FName ModuleThatChanged, EModuleChangeReason ReasonForChange );
+
+	/**
+	 * Updates our module list items
+	 */
+	void UpdateModuleListItems();
+
+	/**
+	 *  module list item filter 
+	 */
+	void OnFilterTextChanged(const FText& InFilterText);
+
+
+
+private:
+
+	typedef SListView< TSharedPtr< FModuleListItem > > SModuleListView;
+	typedef TArray< TSharedPtr< FModuleListItem > > FModuleArray;
+
+	/** List items for the module list */
+	TArray< TSharedPtr< FModuleListItem > > ModuleListItems;
+
+	/** List of all known modules */
+	TSharedPtr< SModuleListView > ModuleListView;
+
+	TSharedPtr< SSearchBox > ModuleNameSearchBox;
+};
diff --git a/Engine/Plugins/Editor/EditorDebugTools/Source/EditorDebugTools/Public/EditorDebugTools.h b/Engine/Plugins/Editor/EditorDebugTools/Source/EditorDebugTools/Public/EditorDebugTools.h
new file mode 100644
--- /dev/null
+++ b/Engine/Plugins/Editor/EditorDebugTools/Source/EditorDebugTools/Public/EditorDebugTools.h
@@ -0,0 +1,30 @@
+// Copyright Epic Games, Inc. All Rights Reserved.
+
+#pragma once
+
+#include "CoreMinimal.h"
+#include "Modules/ModuleManager.h"
+
+class FToolBarBuilder;
+class FMenuBuilder;
+
+class FEditorDebugToolsModule : public IModuleInterface
+{
+public:
+
+	/** IModuleInterface implementation */
+	virtual void StartupModule() override;
+	virtual void ShutdownModule() override;
+	
+	/** This function will be bound to Command (by default it will bring up plugin window) */
+	void PluginButtonClicked();
+	
+private:
+
+	void RegisterMenus();
+
+	TSharedRef<class SDockTab> OnSpawnPluginTab(const class FSpawnTabArgs& SpawnTabArgs);
+
+private:
+	TSharedPtr<class FUICommandList> PluginCommands;
+};
diff --git a/Engine/Plugins/Editor/EditorDebugTools/Source/EditorDebugTools/Public/EditorDebugToolsCommands.h b/Engine/Plugins/Editor/EditorDebugTools/Source/EditorDebugTools/Public/EditorDebugToolsCommands.h
new file mode 100644
--- /dev/null
+++ b/Engine/Plugins/Editor/EditorDebugTools/Source/EditorDebugTools/Public/EditorDebugToolsCommands.h
@@ -0,0 +1,23 @@
+// Copyright Epic Games, Inc. All Rights Reserved.
+
+#pragma once
+
+#include "CoreMinimal.h"
+#include "Framework/Commands/Commands.h"
+#include "EditorDebugToolsStyle.h"
+
+class FEditorDebugToolsCommands : public TCommands<FEditorDebugToolsCommands>
+{
+public:
+
+	FEditorDebugToolsCommands()
+		: TCommands<FEditorDebugToolsCommands>(TEXT("EditorDebugTools"), NSLOCTEXT("Contexts", "EditorDebugTools", "EditorDebugTools Plugin"), NAME_None, FEditorDebugToolsStyle::GetStyleSetName())
+	{
+	}
+
+	// TCommands<> interface
+	virtual void RegisterCommands() override;
+
+public:
+	TSharedPtr< FUICommandInfo > OpenPluginWindow;
+};
\ No newline at end of file
diff --git a/Engine/Plugins/Editor/EditorDebugTools/Source/EditorDebugTools/Public/EditorDebugToolsStyle.h b/Engine/Plugins/Editor/EditorDebugTools/Source/EditorDebugTools/Public/EditorDebugToolsStyle.h
new file mode 100644
--- /dev/null
+++ b/Engine/Plugins/Editor/EditorDebugTools/Source/EditorDebugTools/Public/EditorDebugToolsStyle.h
@@ -0,0 +1,30 @@
+// Copyright Epic Games, Inc. All Rights Reserved.
+
+#pragma once
+
+#include "CoreMinimal.h"
+#include "Styling/SlateStyle.h"
+
+/**  */
+class FEditorDebugToolsStyle
+{
+public:
+
+	static void Initialize();
+
+	static void Shutdown();
+
+
+	/** @return The Slate style set for the Shooter game */
+	static const ISlateStyle& Get();
+
+	static FName GetStyleSetName();
+
+private:
+
+	static TSharedRef< class FSlateStyleSet > Create();
+
+private:
+
+	static TSharedPtr< class FSlateStyleSet > StyleInstance;
+};
\ No newline at end of file
diff --git a/Engine/Source/Developer/GammaUI/GammaUI.Build.cs b/Engine/Source/Developer/GammaUI/GammaUI.Build.cs
deleted file mode 100644
diff --git a/Engine/Source/Developer/GammaUI/Private/GammaUI.cpp b/Engine/Source/Developer/GammaUI/Private/GammaUI.cpp
deleted file mode 100644
diff --git a/Engine/Source/Developer/GammaUI/Private/GammaUIPanel.cpp b/Engine/Source/Developer/GammaUI/Private/GammaUIPanel.cpp
deleted file mode 100644
diff --git a/Engine/Source/Developer/GammaUI/Private/GammaUIPanel.h b/Engine/Source/Developer/GammaUI/Private/GammaUIPanel.h
deleted file mode 100644
diff --git a/Engine/Source/Developer/GammaUI/Public/GammaUI.h b/Engine/Source/Developer/GammaUI/Public/GammaUI.h
deleted file mode 100644
diff --git a/Engine/Source/Developer/ModuleUI/ModuleUI.Build.cs b/Engine/Source/Developer/ModuleUI/ModuleUI.Build.cs
deleted file mode 100644
diff --git a/Engine/Source/Developer/ModuleUI/Private/ModuleUI.cpp b/Engine/Source/Developer/ModuleUI/Private/ModuleUI.cpp
deleted file mode 100644
diff --git a/Engine/Source/Developer/ModuleUI/Private/SModuleUI.cpp b/Engine/Source/Developer/ModuleUI/Private/SModuleUI.cpp
deleted file mode 100644
diff --git a/Engine/Source/Developer/ModuleUI/Private/SModuleUI.h b/Engine/Source/Developer/ModuleUI/Private/SModuleUI.h
deleted file mode 100644
diff --git a/Engine/Source/Developer/ModuleUI/Public/ModuleUIInterface.h b/Engine/Source/Developer/ModuleUI/Public/ModuleUIInterface.h
deleted file mode 100644
diff --git a/Engine/Source/Editor/MainFrame/MainFrame.Build.cs b/Engine/Source/Editor/MainFrame/MainFrame.Build.cs
--- a/Engine/Source/Editor/MainFrame/MainFrame.Build.cs
+++ b/Engine/Source/Editor/MainFrame/MainFrame.Build.cs
@@ -58,7 +58,6 @@ public class MainFrame : ModuleRules
 				"LevelEditor",
 				"Settings",
 				"SourceCodeAccess",
-				"Toolbox",
 				"LocalizationDashboard",
 			}
 		);
diff --git a/Engine/Source/Editor/MainFrame/Private/Menus/MainMenu.cpp b/Engine/Source/Editor/MainFrame/Private/Menus/MainMenu.cpp
--- a/Engine/Source/Editor/MainFrame/Private/Menus/MainMenu.cpp
+++ b/Engine/Source/Editor/MainFrame/Private/Menus/MainMenu.cpp
@@ -31,8 +31,6 @@
 #include "Features/IModularFeatures.h"
 #include "UndoHistoryModule.h"
 #include "Framework/Commands/GenericCommands.h"
-#include "ToolboxModule.h"
-
 
 #define LOCTEXT_NAMESPACE "MainFileMenu"
 
diff --git a/Engine/Source/Editor/TurnkeySupport/TurnkeySupport.Build.cs b/Engine/Source/Editor/TurnkeySupport/TurnkeySupport.Build.cs
--- a/Engine/Source/Editor/TurnkeySupport/TurnkeySupport.Build.cs
+++ b/Engine/Source/Editor/TurnkeySupport/TurnkeySupport.Build.cs
@@ -36,7 +36,6 @@ public class TurnkeySupport : ModuleRules
 				"Analytics",
 				"ToolMenus",
 				"LauncherServices",
-				
 				"SettingsEditor"
 			}
 		);
@@ -50,7 +49,6 @@ public class TurnkeySupport : ModuleRules
 				"LevelEditor",
 				"Settings",
 				"SourceCodeAccess",
-				"Toolbox",
 				"LocalizationDashboard",
 			}
 		);
diff --git a/Engine/Source/Editor/UnrealEd/Private/EditorEngine.cpp b/Engine/Source/Editor/UnrealEd/Private/EditorEngine.cpp
--- a/Engine/Source/Editor/UnrealEd/Private/EditorEngine.cpp
+++ b/Engine/Source/Editor/UnrealEd/Private/EditorEngine.cpp
@@ -989,14 +989,11 @@ void UEditorEngine::Init(IEngineLoop* InEngineLoop)
 			TEXT("Documentation"),
 			TEXT("WorkspaceMenuStructure"),
 			TEXT("MainFrame"),
-			TEXT("GammaUI"),
 			TEXT("OutputLog"),
 			TEXT("SourceControl"),
 			TEXT("TextureCompressor"),
 			TEXT("MeshUtilities"),
 			TEXT("MovieSceneTools"),
-			TEXT("ModuleUI"),
-			TEXT("Toolbox"),
 			TEXT("ClassViewer"),
 			TEXT("StructViewer"),
 			TEXT("ContentBrowser"),
diff --git a/Engine/Source/Runtime/Launch/Launch.Build.cs b/Engine/Source/Runtime/Launch/Launch.Build.cs
--- a/Engine/Source/Runtime/Launch/Launch.Build.cs
+++ b/Engine/Source/Runtime/Launch/Launch.Build.cs
@@ -196,9 +196,6 @@ public class Launch : ModuleRules
 			DynamicallyLoadedModuleNames.AddRange(new string[] {
 					"AutomationWindow",
 					"ProfilerClient",
-					"Toolbox",
-					"GammaUI",
-					"ModuleUI",
 					"OutputLog",
 					"TextureCompressor",
 					"MeshUtilities",
diff --git a/Engine/Source/Runtime/Slate/Private/Framework/Application/SlateApplication.cpp b/Engine/Source/Runtime/Slate/Private/Framework/Application/SlateApplication.cpp
--- a/Engine/Source/Runtime/Slate/Private/Framework/Application/SlateApplication.cpp
+++ b/Engine/Source/Runtime/Slate/Private/Framework/Application/SlateApplication.cpp
@@ -36,7 +36,6 @@
 #include "Framework/Notifications/SlateAsyncTaskNotificationImpl.h"
 #include "Framework/Application/IInputProcessor.h"
 #include "GenericPlatform/ITextInputMethodSystem.h"
-#include "ToolboxModule.h"
 #include "Framework/Docking/TabCommands.h"
 #include "HAL/LowLevelMemTracker.h"
 #include "Math/UnitConversion.h"
@@ -4293,20 +4292,6 @@ bool FSlateApplication::ProcessKeyDownEvent( const FKeyEvent& InKeyEvent )
 			}
 		}
 #endif
-
-#if !(UE_BUILD_SHIPPING || UE_BUILD_TEST)
-		// Ctrl+Shift+~ summons the Toolbox.
-		if (InKeyEvent.GetKey() == EKeys::Tilde && InKeyEvent.IsControlDown() && InKeyEvent.IsShiftDown())
-		{
-			IToolboxModule* ToolboxModule = FModuleManager::LoadModulePtr<IToolboxModule>("Toolbox");
-			if (ToolboxModule)
-			{
-				ToolboxModule->SummonToolbox();
-			}
-		}
-
-#endif //!(UE_BUILD_SHIPPING || UE_BUILD_TEST)
-
 		// Bubble the keyboard event
 		TSharedRef<FWidgetPath> EventPathRef = SlateUser->GetFocusPath();
 		const FWidgetPath& EventPath = EventPathRef.Get();
diff --git a/Engine/Source/Runtime/Slate/Slate.Build.cs b/Engine/Source/Runtime/Slate/Slate.Build.cs
--- a/Engine/Source/Runtime/Slate/Slate.Build.cs
+++ b/Engine/Source/Runtime/Slate/Slate.Build.cs
@@ -41,8 +41,6 @@ public class Slate : ModuleRules
 				"Runtime/Slate/Private/Widgets/Notifications",
 				"Runtime/Slate/Private/Widgets/Text",
 				"Runtime/Slate/Private/Widgets/Views",
-
-				"Runtime/Toolbox/Public/"
 			});
 
 		if ((Target.Platform == UnrealTargetPlatform.Win64) ||
diff --git a/Engine/Source/Runtime/Toolbox/Private/ToolboxModule.cpp b/Engine/Source/Runtime/Toolbox/Private/ToolboxModule.cpp
--- a/Engine/Source/Runtime/Toolbox/Private/ToolboxModule.cpp
+++ b/Engine/Source/Runtime/Toolbox/Private/ToolboxModule.cpp
@@ -18,206 +18,15 @@
 #include "DesktopPlatformModule.h"
 #include "Widgets/Docking/SDockTab.h"
 #include "Widgets/Testing/STestSuite.h"
+#include "Widgets/Testing/SStarshipSuite.h"
 #include "Editor/WorkspaceMenuStructure/Public/WorkspaceMenuStructure.h"
 #include "Editor/WorkspaceMenuStructure/Public/WorkspaceMenuStructureModule.h"
 #include "ISourceCodeAccessor.h"
 #include "ISourceCodeAccessModule.h"
+#include "Misc/CommandLine.h"
 
 #if !UE_BUILD_SHIPPING
 #include "ISlateReflectorModule.h"
 #endif // #if !UE_BUILD_SHIPPING
 
 
-static bool bTabsRegistered = false;
-
-static bool CanShowModulesTab()
-{
-	ISourceCodeAccessModule* SourceCodeAccessModule = FModuleManager::GetModulePtr<ISourceCodeAccessModule>("SourceCodeAccess");
-	return SourceCodeAccessModule != nullptr && SourceCodeAccessModule->GetAccessor().CanAccessSourceCode();
-}
-
-class SDebugPanel : public SCompoundWidget
-{
-public:
-
-	SLATE_BEGIN_ARGS(SDebugPanel){}
-	SLATE_END_ARGS()
-
-	BEGIN_SLATE_FUNCTION_BUILD_OPTIMIZATION
-	void Construct( const FArguments& )
-	{
-		ChildSlot
-		[
-			SNew( SVerticalBox )
-			+ SVerticalBox::Slot()
-			.AutoHeight()
-			.Padding( 4.0f )
-			.HAlign(HAlign_Left)
-			[
-				SNew( SButton )
-				.Text( NSLOCTEXT("DeveloperToolbox", "ReloadTextures", "Reload Textures") )
-				.OnClicked( this, &SDebugPanel::OnReloadTexturesClicked )
-			]
-			
-			+ SVerticalBox::Slot()
-			.AutoHeight()
-			.Padding( 4.0f )
-			.HAlign(HAlign_Left)
-			[
-				SNew( SButton )
-				.Text( NSLOCTEXT("DeveloperToolbox", "FlushFontCache", "Flush Font Cache") )
-				.OnClicked( this, &SDebugPanel::OnFlushFontCacheClicked )
-			]
-			+ SVerticalBox::Slot()
-			.AutoHeight()
-			.Padding( 4.0f )
-			.HAlign(HAlign_Left)
-			[
-				SNew( SButton )
-				.Text( NSLOCTEXT("DeveloperToolbox", "TestSuite", "Test Suite") )
-				.OnClicked( this, &SDebugPanel::OnTestSuiteClicked )
-			]
-			+ SVerticalBox::Slot()
-			.AutoHeight()
-			.Padding( 4.0f )
-			.HAlign(HAlign_Left)
-			[
-				SNew( SButton )
-				.Text( NSLOCTEXT("DeveloperToolbox", "DisplayTextureAtlases", "Display Texture Atlases") )
-				.OnClicked( this, &SDebugPanel::OnDisplayTextureAtlases )
-			]
-			+ SVerticalBox::Slot()
-			.AutoHeight()
-			.Padding( 4.0f )
-			.HAlign(HAlign_Left)
-			[
-				SNew( SButton )
-				.Text( NSLOCTEXT("DeveloperToolbox", "DisplayFontAtlases", "Display Font Atlases") )
-				.OnClicked( this, &SDebugPanel::OnDisplayFontAtlases )
-			]
-		];
-	}
-	END_SLATE_FUNCTION_BUILD_OPTIMIZATION
-
-	
-	FReply OnReloadTexturesClicked()
-	{
-		FSlateApplication::Get().GetRenderer()->ReloadTextureResources();
-		return FReply::Handled();
-	}
-
-	FReply OnDisplayTextureAtlases()
-	{
-#if !UE_BUILD_SHIPPING
-		static const FName SlateReflectorModuleName("SlateReflector");
-		FModuleManager::LoadModuleChecked<ISlateReflectorModule>(SlateReflectorModuleName).DisplayTextureAtlasVisualizer();
-#endif // #if !UE_BUILD_SHIPPING
-		return FReply::Handled();
-	}
-
-	FReply OnDisplayFontAtlases()
-	{
-#if !UE_BUILD_SHIPPING
-		static const FName SlateReflectorModuleName("SlateReflector");
-		FModuleManager::LoadModuleChecked<ISlateReflectorModule>(SlateReflectorModuleName).DisplayFontAtlasVisualizer();
-#endif // #if !UE_BUILD_SHIPPING
-		return FReply::Handled();
-	}
-
-	FReply OnFlushFontCacheClicked()
-	{
-		FSlateApplication::Get().GetRenderer()->FlushFontCache(TEXT("SDebugPanel::OnFlushFontCacheClicked"));
-		return FReply::Handled();
-	}
-
-	FReply OnTestSuiteClicked()
-	{
-#if !UE_BUILD_SHIPPING
-		RestoreSlateTestSuite();
-#endif
-		return FReply::Handled();
-	}
-};
-
-TSharedRef<SDockTab> CreateDebugToolsTab( const FSpawnTabArgs& Args )
-{
-	FGammaUI& GammaUI = FModuleManager::Get().LoadModuleChecked<FGammaUI>( FName("GammaUI") );
-
-	return
-		SNew(SDockTab)
-		.TabRole( ETabRole::NomadTab )
-		[
-			SNew(SVerticalBox)
-			+SVerticalBox::Slot()
-			.AutoHeight()
-			[
-				SNew( SDebugPanel )
-			]
-			+SVerticalBox::Slot()
-			.AutoHeight()
-			[
-				GammaUI.GetGammaUIPanel().ToSharedRef()
-			]
-		];
-}
-
-TSharedRef<SDockTab> CreateModulesTab( const FSpawnTabArgs& Args )
-{
-	IModuleUIInterface& ModuleUI = FModuleManager::Get().LoadModuleChecked<IModuleUIInterface>( FName("ModuleUI") );
-
-	return
-		SNew(SDockTab)
-		.TabRole( ETabRole::NomadTab )
-		[
-			SNew(SBorder)
-			.Padding(2)
-			.BorderImage( FEditorStyle::GetBrush("ToolPanel.GroupBorder") )
-			[
-				ModuleUI.GetModuleUIWidget()
-			]
-		];
-}
-
-
-
-class FToolboxModule : public IToolboxModule
-{
-	void StartupModule() override
-	{
-		FTabSpawnerEntry& Spawner = FGlobalTabmanager::Get()->RegisterNomadTabSpawner("DebugTools", FOnSpawnTab::CreateStatic(&CreateDebugToolsTab))
-			.SetDisplayName(NSLOCTEXT("Toolbox", "DebugTools", "Debug Tools"))
-			.SetTooltipText(NSLOCTEXT("Toolbox", "DebugToolsTooltipText", "Open the Debug Tools tab."))
-			.SetIcon(FSlateIcon(FEditorStyle::GetStyleSetName(), "DebugTools.TabIcon"))
-			.SetGroup(WorkspaceMenu::GetMenuStructure().GetDeveloperToolsDebugCategory());
-
-		if (CanShowModulesTab())
-		{
-			FGlobalTabmanager::Get()->RegisterNomadTabSpawner("ModulesTab", FOnSpawnTab::CreateStatic(&CreateModulesTab))
-				.SetDisplayName(NSLOCTEXT("Toolbox", "Modules", "Modules"))
-				.SetTooltipText(NSLOCTEXT("Toolbox", "ModulesTooltipText", "Open the Modules tab."))
-				.SetIcon(FSlateIcon(FEditorStyle::GetStyleSetName(), "Modules.TabIcon"))
-				.SetGroup(WorkspaceMenu::GetMenuStructure().GetDeveloperToolsMiscCategory());
-		}
-	}
-
-	void ShutdownModule() override
-	{
-		if (FSlateApplication::IsInitialized())
-		{
-			FGlobalTabmanager::Get()->UnregisterNomadTabSpawner("DebugTools");
-			if (CanShowModulesTab())
-			{
-				FGlobalTabmanager::Get()->UnregisterNomadTabSpawner("ModulesTab");
-			}
-		}
-	}
-
-	virtual void SummonToolbox() override
-	{
-		FGlobalTabmanager::Get()->TryInvokeTab(FTabId("DebugTools"));
-	}
-};
-
-
-IMPLEMENT_MODULE(FToolboxModule, Toolbox);
-
diff --git a/Engine/Source/Runtime/Toolbox/Public/ToolboxModule.h b/Engine/Source/Runtime/Toolbox/Public/ToolboxModule.h
deleted file mode 100644
diff --git a/Engine/Source/Runtime/Toolbox/Toolbox.Build.cs b/Engine/Source/Runtime/Toolbox/Toolbox.Build.cs
deleted file mode 100644
