Move debug tools icons to its own plugin

diff --git a/Engine/Plugins/Editor/EditorDebugTools/Source/EditorDebugTools/Private/EditorDebugTools.cpp b/Engine/Plugins/Editor/EditorDebugTools/Source/EditorDebugTools/Private/EditorDebugTools.cpp
--- a/Engine/Plugins/Editor/EditorDebugTools/Source/EditorDebugTools/Private/EditorDebugTools.cpp
+++ b/Engine/Plugins/Editor/EditorDebugTools/Source/EditorDebugTools/Private/EditorDebugTools.cpp
@@ -63,10 +63,12 @@ TSharedRef<SDockTab> CreateModulesTab(const FSpawnTabArgs& Args)
 
 void FEditorDebugToolsModule::StartupModule()
 {
+	FEditorDebugToolsStyle::Initialize();
+
 	FTabSpawnerEntry& Spawner = FGlobalTabmanager::Get()->RegisterNomadTabSpawner(EditorDebugToolsTabName, FOnSpawnTab::CreateStatic(&CreateDebugToolsTab))
 		.SetDisplayName(NSLOCTEXT("Toolbox", "DebugTools", "Debug Tools"))
 		.SetTooltipText(NSLOCTEXT("Toolbox", "DebugToolsTooltipText", "Open the Debug Tools tab."))
-		.SetIcon(FSlateIcon(FEditorStyle::GetStyleSetName(), "DebugTools.TabIcon"))
+		.SetIcon(FSlateIcon(FEditorDebugToolsStyle::GetStyleSetName(), "DebugTools.TabIcon"))
 		.SetGroup(WorkspaceMenu::GetMenuStructure().GetDeveloperToolsDebugCategory());
 
 	if (CanShowModulesTab())
@@ -74,13 +76,15 @@ void FEditorDebugToolsModule::StartupModule()
 		FGlobalTabmanager::Get()->RegisterNomadTabSpawner(ModulesTabName, FOnSpawnTab::CreateStatic(&CreateModulesTab))
 			.SetDisplayName(NSLOCTEXT("Toolbox", "Modules", "Modules"))
 			.SetTooltipText(NSLOCTEXT("Toolbox", "ModulesTooltipText", "Open the Modules tab."))
-			.SetIcon(FSlateIcon(FEditorStyle::GetStyleSetName(), "Modules.TabIcon"))
+			.SetIcon(FSlateIcon(FEditorDebugToolsStyle::GetStyleSetName(), "Modules.TabIcon"))
 			.SetGroup(WorkspaceMenu::GetMenuStructure().GetDeveloperToolsMiscCategory());
 	}
 }
 
 void FEditorDebugToolsModule::ShutdownModule()
 {
+	FEditorDebugToolsStyle::Shutdown();
+
 	FGlobalTabmanager::Get()->UnregisterNomadTabSpawner(EditorDebugToolsTabName);
 
 	if (CanShowModulesTab())
diff --git a/Engine/Plugins/Editor/EditorDebugTools/Source/EditorDebugTools/Private/EditorDebugToolsStyle.cpp b/Engine/Plugins/Editor/EditorDebugTools/Source/EditorDebugTools/Private/EditorDebugToolsStyle.cpp
--- a/Engine/Plugins/Editor/EditorDebugTools/Source/EditorDebugTools/Private/EditorDebugToolsStyle.cpp
+++ b/Engine/Plugins/Editor/EditorDebugTools/Source/EditorDebugTools/Private/EditorDebugToolsStyle.cpp
@@ -4,9 +4,12 @@
 #include "Styling/SlateStyleRegistry.h"
 #include "Framework/Application/SlateApplication.h"
 #include "Interfaces/IPluginManager.h"
+#include "Styling/SlateStyleMacros.h"
 
 TSharedPtr< FSlateStyleSet > FEditorDebugToolsStyle::StyleInstance = NULL;
 
+#define IMAGE_PLUGIN_BRUSH( RelativePath, ... ) FSlateImageBrush( Style->RootToContentDir( RelativePath, TEXT(".png") ), __VA_ARGS__ )
+
 void FEditorDebugToolsStyle::Initialize()
 {
 	if (!StyleInstance.IsValid())
@@ -37,7 +40,16 @@ const FVector2D Icon40x40(40.0f, 40.0f);
 TSharedRef< FSlateStyleSet > FEditorDebugToolsStyle::Create()
 {
 	TSharedRef< FSlateStyleSet > Style = MakeShareable(new FSlateStyleSet("EditorDebugToolsStyle"));
-	Style->SetContentRoot(IPluginManager::Get().FindPlugin("EditorDebugTools")->GetBaseDir() / TEXT("Resources"));
+	Style->SetContentRoot(IPluginManager::Get().FindPlugin("EditorDebugTools")->GetBaseDir() / TEXT("Resources/Slate"));
+
+	// Gamma reference.
+	Style->Set("GammaReference", new IMAGE_PLUGIN_BRUSH("GammaReference", FVector2D(256, 128)));
+
+	// Debug Tools Window
+	Style->Set("DebugTools.TabIcon", new IMAGE_PLUGIN_BRUSH("icon_tab_DebugTools_16x", Icon16x16));
+	
+	// Modules Window
+	Style->Set("Modules.TabIcon", new IMAGE_PLUGIN_BRUSH("icon_tab_Modules_16px", Icon16x16));
 
 	return Style;
 }
diff --git a/Engine/Plugins/Editor/EditorDebugTools/Source/EditorDebugTools/Private/GammaUIPanel.cpp b/Engine/Plugins/Editor/EditorDebugTools/Source/EditorDebugTools/Private/GammaUIPanel.cpp
--- a/Engine/Plugins/Editor/EditorDebugTools/Source/EditorDebugTools/Private/GammaUIPanel.cpp
+++ b/Engine/Plugins/Editor/EditorDebugTools/Source/EditorDebugTools/Private/GammaUIPanel.cpp
@@ -8,25 +8,24 @@
 #include "Widgets/Text/STextBlock.h"
 #include "Widgets/Input/SSpinBox.h"
 #include "EditorStyleSet.h"
+#include "EditorDebugToolsStyle.h"
 
 void SGammaUIPanel::Construct(const SGammaUIPanel::FArguments& InArgs)
 {
 	ChildSlot
-		.Padding( FMargin(8) )
+	.Padding( FMargin(8) )
 	[
 		SNew( SVerticalBox )
-
 		+ SVerticalBox::Slot()
-			.AutoHeight()
-			.Padding( 0.0f, 4.0f, 0.0f, 4.0f )
+		.AutoHeight()
+		.Padding(0.0f, 4.0f, 0.0f, 4.0f)
 		[
-			SNew( STextBlock )
-				.Text( NSLOCTEXT( "GammaUI", "GammaUILabel", "Gamma" ) )
+			SNew(STextBlock)
+			.Text(NSLOCTEXT("GammaUI", "GammaUILabel", "Gamma"))
 		]
-
 		+ SVerticalBox::Slot()
-			.AutoHeight()
-			.Padding( 0.0f, 4.0f, 0.0f, 4.0f )
+		.AutoHeight()
+		.Padding( 0.0f, 4.0f, 0.0f, 4.0f )
 		[
 			SNew( SSpinBox<float> )
 			.Delta(0.01f)
@@ -35,14 +34,13 @@ void SGammaUIPanel::Construct(const SGammaUIPanel::FArguments& InArgs)
 			.Value( this, &SGammaUIPanel::OnGetGamma )
 			.OnValueChanged( this, &SGammaUIPanel::OnGammaChanged )
 		]
-
 		+ SVerticalBox::Slot()
-			.AutoHeight()
-			.Padding( 0.0f, 4.0f, 0.0f, 4.0f )
-			.HAlign( HAlign_Center )
+		.AutoHeight()
+		.Padding(0.0f, 4.0f, 0.0f, 4.0f)
+		.HAlign(HAlign_Center)
 		[
 			SNew(SImage)
-			.Image(FEditorStyle::GetBrush(TEXT("GammaReference")))
+			.Image(FEditorDebugToolsStyle::Get().GetBrush(TEXT("GammaReference")))
 		]
 	];
 }
