Added missing copyright notice

diff --git a/Engine/Plugins/Editor/EditorDebugTools/Source/EditorDebugTools/Private/SDebugPanel.cpp b/Engine/Plugins/Editor/EditorDebugTools/Source/EditorDebugTools/Private/SDebugPanel.cpp
--- a/Engine/Plugins/Editor/EditorDebugTools/Source/EditorDebugTools/Private/SDebugPanel.cpp
+++ b/Engine/Plugins/Editor/EditorDebugTools/Source/EditorDebugTools/Private/SDebugPanel.cpp
@@ -1,3 +1,5 @@
+// Copyright Epic Games, Inc. All Rights Reserved.
+
 #include "SDebugPanel.h"
 #include "ISlateReflectorModule.h"
 #include "Widgets/Testing/SStarshipSuite.h"
