Fix-up VS generated include paths.

diff --git a/Engine/Plugins/Editor/EditorDebugTools/Source/EditorDebugTools/Private/SDebugPanel.cpp b/Engine/Plugins/Editor/EditorDebugTools/Source/EditorDebugTools/Private/SDebugPanel.cpp
--- a/Engine/Plugins/Editor/EditorDebugTools/Source/EditorDebugTools/Private/SDebugPanel.cpp
+++ b/Engine/Plugins/Editor/EditorDebugTools/Source/EditorDebugTools/Private/SDebugPanel.cpp
@@ -4,11 +4,11 @@
 #include "ISlateReflectorModule.h"
 #include "Widgets/Testing/SStarshipSuite.h"
 #include "Widgets/Testing/STestSuite.h"
-#include "Widgets\SBoxPanel.h"
-#include "Widgets\Input\SButton.h"
-#include "Styling\CoreStyle.h"
-#include "Modules\ModuleManager.h"
-#include "Framework\Application\SlateApplication.h"
+#include "Widgets/SBoxPanel.h"
+#include "Widgets/Input/SButton.h"
+#include "Styling/CoreStyle.h"
+#include "Modules/ModuleManager.h"
+#include "Framework/Application/SlateApplication.h"
 
 void SDebugPanel::Construct(const FArguments& InArgs)
 {
