Fix CIS error 'IMAGE_BRUSH' macro redefined
#rnx

diff --git a/Engine/Plugins/Editor/EditorDebugTools/Source/EditorDebugTools/Private/EditorDebugToolsStyle.cpp b/Engine/Plugins/Editor/EditorDebugTools/Source/EditorDebugTools/Private/EditorDebugToolsStyle.cpp
--- a/Engine/Plugins/Editor/EditorDebugTools/Source/EditorDebugTools/Private/EditorDebugToolsStyle.cpp
+++ b/Engine/Plugins/Editor/EditorDebugTools/Source/EditorDebugTools/Private/EditorDebugToolsStyle.cpp
@@ -8,7 +8,7 @@
 
 TSharedPtr< FSlateStyleSet > FEditorDebugToolsStyle::StyleInstance = NULL;
 
-#define IMAGE_BRUSH( RelativePath, ... ) FSlateImageBrush( FPaths::EngineContentDir() / "Editor/Slate"/ RelativePath + TEXT(".png"), __VA_ARGS__ )
+#define IMAGE_BRUSH( RelativePath, ... ) FSlateImageBrush( Style->RootToContentDir( RelativePath, TEXT(".png") ), __VA_ARGS__ )
 #define IMAGE_PLUGIN_BRUSH( RelativePath, ... ) FSlateImageBrush( Style->RootToContentDir( RelativePath, TEXT(".png") ), __VA_ARGS__ )
 
 void FEditorDebugToolsStyle::Initialize()
@@ -61,3 +61,6 @@ const ISlateStyle& FEditorDebugToolsStyle::Get()
 {
 	return *StyleInstance;
 }
+
+#undef IMAGE_BRUSH
+#undef IMAGE_PLUGIN_BRUSH
