Fix CIS error 'IMAGE_BRUSH' macro redefined
#rnx

diff --git a/Engine/Plugins/Editor/EditorDebugTools/Source/EditorDebugTools/Private/EditorDebugToolsStyle.cpp b/Engine/Plugins/Editor/EditorDebugTools/Source/EditorDebugTools/Private/EditorDebugToolsStyle.cpp
--- a/Engine/Plugins/Editor/EditorDebugTools/Source/EditorDebugTools/Private/EditorDebugToolsStyle.cpp
+++ b/Engine/Plugins/Editor/EditorDebugTools/Source/EditorDebugTools/Private/EditorDebugToolsStyle.cpp
@@ -8,7 +8,7 @@
 
 TSharedPtr< FSlateStyleSet > FEditorDebugToolsStyle::StyleInstance = NULL;
 
-#define IMAGE_BRUSH( RelativePath, ... ) FSlateImageBrush( Style->RootToContentDir( RelativePath, TEXT(".png") ), __VA_ARGS__ )
+#define LOCAL_IMAGE_BRUSH( RelativePath, ... ) FSlateImageBrush( FPaths::EngineContentDir() / "Editor/Slate"/ RelativePath + TEXT(".png"), __VA_ARGS__ )
 #define IMAGE_PLUGIN_BRUSH( RelativePath, ... ) FSlateImageBrush( Style->RootToContentDir( RelativePath, TEXT(".png") ), __VA_ARGS__ )
 
 void FEditorDebugToolsStyle::Initialize()
@@ -47,7 +47,7 @@ TSharedRef< FSlateStyleSet > FEditorDebugToolsStyle::Create()
 	Style->Set("GammaReference", new IMAGE_PLUGIN_BRUSH("GammaReference", FVector2D(256, 128)));
 
 	// Debug Tools Window
-	Style->Set("DebugTools.TabIcon", new IMAGE_BRUSH("Icons/icon_tab_DebugTools_16x", Icon16x16));
+	Style->Set("DebugTools.TabIcon", new LOCAL_IMAGE_BRUSH("Icons/icon_tab_DebugTools_16x", Icon16x16));
 	
 	// Modules Window
 	Style->Set("Modules.TabIcon", new IMAGE_PLUGIN_BRUSH("icon_tab_Modules_16px", Icon16x16));
@@ -62,5 +62,5 @@ const ISlateStyle& FEditorDebugToolsStyle::Get()
 	return *StyleInstance;
 }
 
-#undef IMAGE_BRUSH
+#undef LOCAL_IMAGE_BRUSH
 #undef IMAGE_PLUGIN_BRUSH
