Fixed issue where modules tab wouldn't reliably be available.

diff --git a/Engine/Plugins/Editor/EditorDebugTools/Source/EditorDebugTools/Private/EditorDebugTools.cpp b/Engine/Plugins/Editor/EditorDebugTools/Source/EditorDebugTools/Private/EditorDebugTools.cpp
--- a/Engine/Plugins/Editor/EditorDebugTools/Source/EditorDebugTools/Private/EditorDebugTools.cpp
+++ b/Engine/Plugins/Editor/EditorDebugTools/Source/EditorDebugTools/Private/EditorDebugTools.cpp
@@ -8,8 +8,6 @@
 #include "Widgets/Text/STextBlock.h"
 #include "ToolMenus.h"
 #include "SDebugPanel.h"
-#include "ISourceCodeAccessModule.h"
-#include "ISourceCodeAccessor.h"
 #include "GammaUIPanel.h"
 #include "SModuleUI.h"
 #include "WorkspaceMenuStructureModule.h"
@@ -18,12 +16,6 @@
 static const FName EditorDebugToolsTabName("DebugTools");
 static const FName ModulesTabName("Modules");
 
-static bool CanShowModulesTab()
-{
-	ISourceCodeAccessModule* SourceCodeAccessModule = FModuleManager::GetModulePtr<ISourceCodeAccessModule>("SourceCodeAccess");
-	return SourceCodeAccessModule != nullptr && SourceCodeAccessModule->GetAccessor().CanAccessSourceCode();
-}
-
 #define LOCTEXT_NAMESPACE "FEditorDebugToolsModule"
 
 TSharedRef<SDockTab> CreateDebugToolsTab(const FSpawnTabArgs& Args)
@@ -71,14 +63,11 @@ void FEditorDebugToolsModule::StartupModule()
 		.SetIcon(FSlateIcon(FAppStyle::GetAppStyleSetName(), "Debug"))
 		.SetGroup(WorkspaceMenu::GetMenuStructure().GetDeveloperToolsDebugCategory());
 
-	if (CanShowModulesTab())
-	{
-		FGlobalTabmanager::Get()->RegisterNomadTabSpawner(ModulesTabName, FOnSpawnTab::CreateStatic(&CreateModulesTab))
-			.SetDisplayName(NSLOCTEXT("Toolbox", "Modules", "Modules"))
-			.SetTooltipText(NSLOCTEXT("Toolbox", "ModulesTooltipText", "Open the Modules tab."))
-			.SetIcon(FSlateIcon(FAppStyle::GetAppStyleSetName(), "Modules"))
-			.SetGroup(WorkspaceMenu::GetMenuStructure().GetDeveloperToolsDebugCategory());
-	}
+	FGlobalTabmanager::Get()->RegisterNomadTabSpawner(ModulesTabName, FOnSpawnTab::CreateStatic(&CreateModulesTab))
+		.SetDisplayName(NSLOCTEXT("Toolbox", "Modules", "Modules"))
+		.SetTooltipText(NSLOCTEXT("Toolbox", "ModulesTooltipText", "Open the Modules tab."))
+		.SetIcon(FSlateIcon(FAppStyle::GetAppStyleSetName(), "Modules"))
+		.SetGroup(WorkspaceMenu::GetMenuStructure().GetDeveloperToolsDebugCategory());
 }
 
 void FEditorDebugToolsModule::ShutdownModule()
@@ -87,10 +76,7 @@ void FEditorDebugToolsModule::ShutdownModule()
 
 	FGlobalTabmanager::Get()->UnregisterNomadTabSpawner(EditorDebugToolsTabName);
 
-	if (CanShowModulesTab())
-	{
-		FGlobalTabmanager::Get()->UnregisterNomadTabSpawner(ModulesTabName);
-	}
+	FGlobalTabmanager::Get()->UnregisterNomadTabSpawner(ModulesTabName);
 }
 
 
diff --git a/Engine/Plugins/Editor/EditorDebugTools/Source/EditorDebugTools/Private/SModuleUI.cpp b/Engine/Plugins/Editor/EditorDebugTools/Source/EditorDebugTools/Private/SModuleUI.cpp
--- a/Engine/Plugins/Editor/EditorDebugTools/Source/EditorDebugTools/Private/SModuleUI.cpp
+++ b/Engine/Plugins/Editor/EditorDebugTools/Source/EditorDebugTools/Private/SModuleUI.cpp
@@ -13,6 +13,14 @@
 #include "Widgets/Views/SListView.h"
 #include "Misc/HotReloadInterface.h"
 #include "Widgets/Input/SSearchBox.h"
+#include "ISourceCodeAccessModule.h"
+#include "ISourceCodeAccessor.h"
+
+static bool CanAccessSourceCode()
+{
+	ISourceCodeAccessModule* SourceCodeAccessModule = FModuleManager::GetModulePtr<ISourceCodeAccessModule>("SourceCodeAccess");
+	return SourceCodeAccessModule != nullptr && SourceCodeAccessModule->GetAccessor().CanAccessSourceCode();
+}
 
 BEGIN_SLATE_FUNCTION_BUILD_OPTIMIZATION
 void SModuleUI::Construct(const SModuleUI::FArguments& InArgs)
@@ -328,6 +336,11 @@ EVisibility SModuleUI::FModuleListItem::GetVisibilityBasedOnRecompilableState()
 		return EVisibility::Hidden;
 	}
 
+	if (!CanAccessSourceCode())
+	{
+		return EVisibility::Hidden;
+	}
+
 	const bool bIsHotReloadable = FModuleManager::Get().DoesLoadedModuleHaveUObjects(ModuleName);
 	const bool bCanReload = ( !FModuleManager::Get().IsModuleLoaded(ModuleName)
 							|| FModuleManager::Get().GetModule(ModuleName)->SupportsDynamicReloading()
diff --git a/Engine/Plugins/Editor/EditorDebugTools/Source/EditorDebugTools/Public/EditorDebugTools.h b/Engine/Plugins/Editor/EditorDebugTools/Source/EditorDebugTools/Public/EditorDebugTools.h
--- a/Engine/Plugins/Editor/EditorDebugTools/Source/EditorDebugTools/Public/EditorDebugTools.h
+++ b/Engine/Plugins/Editor/EditorDebugTools/Source/EditorDebugTools/Public/EditorDebugTools.h
@@ -15,16 +15,4 @@ public:
 	/** IModuleInterface implementation */
 	virtual void StartupModule() override;
 	virtual void ShutdownModule() override;
-	
-	/** This function will be bound to Command (by default it will bring up plugin window) */
-	void PluginButtonClicked();
-	
-private:
-
-	void RegisterMenus();
-
-	TSharedRef<class SDockTab> OnSpawnPluginTab(const class FSpawnTabArgs& SpawnTabArgs);
-
-private:
-	TSharedPtr<class FUICommandList> PluginCommands;
 };
