Remove the Starship Gallery test suite from the editor. It was not designed to run inside the editor and can therefore cause crashes.

diff --git a/Engine/Plugins/Editor/EditorDebugTools/Source/EditorDebugTools/Private/SDebugPanel.cpp b/Engine/Plugins/Editor/EditorDebugTools/Source/EditorDebugTools/Private/SDebugPanel.cpp
--- a/Engine/Plugins/Editor/EditorDebugTools/Source/EditorDebugTools/Private/SDebugPanel.cpp
+++ b/Engine/Plugins/Editor/EditorDebugTools/Source/EditorDebugTools/Private/SDebugPanel.cpp
@@ -41,15 +41,6 @@ void SDebugPanel::Construct(const FArguments& InArgs)
 		.AutoHeight()
 		.Padding( 4.0f )
 		.HAlign(HAlign_Left)
-		[
-			SNew( SButton )
-			.Text( NSLOCTEXT("DeveloperToolbox", "TestSuite", "Test Suite") )
-			.OnClicked( this, &SDebugPanel::OnTestSuiteClicked )
-		]
-		+ SVerticalBox::Slot()
-		.AutoHeight()
-		.Padding( 4.0f )
-		.HAlign(HAlign_Left)
 		[
 			SNew( SButton )
 			.Text( NSLOCTEXT("DeveloperToolbox", "DisplayTextureAtlases", "Display Texture Atlases") )
@@ -96,18 +87,3 @@ FReply SDebugPanel::OnFlushFontCacheClicked()
 	FSlateApplication::Get().GetRenderer()->FlushFontCache(TEXT("SDebugPanel::OnFlushFontCacheClicked"));
 	return FReply::Handled();
 }
-
-FReply SDebugPanel::OnTestSuiteClicked()
-{
-#if !UE_BUILD_SHIPPING
-	if (FCoreStyle::IsStarshipStyle())
-	{
-		RestoreStarshipSuite();
-	}
-	else
-	{
-		RestoreSlateTestSuite();
-	}
-#endif
-	return FReply::Handled();
-}
